#ttp 309441 - TPS: make sure that libpng allocates memory via engine allocator
#proj Editor.UnrealEd
#branch UE4
#summary libPNG allocations now routing via FMemory correctly
#change Swapped from png_create_read_struct to png_create_read_struct_2 and supplied (wrapped) calls to FMemory::malloc/free

diff --git a/Engine/Source/Developer/ImageWrapper/Private/PngImageWrapper.cpp b/Engine/Source/Developer/ImageWrapper/Private/PngImageWrapper.cpp
--- a/Engine/Source/Developer/ImageWrapper/Private/PngImageWrapper.cpp
+++ b/Engine/Source/Developer/ImageWrapper/Private/PngImageWrapper.cpp
@@ -205,7 +205,7 @@ void FPngImageWrapper::UncompressPNGData( const ERGBFormat::Type InFormat, const
 	// Reset to the beginning of file so we can use png_read_png(), which expects to start at the beginning.
 	ReadOffset = 0;
 		
-	png_structp png_ptr	= png_create_read_struct( PNG_LIBPNG_VER_STRING, this, FPngImageWrapper::user_error_fn, FPngImageWrapper::user_warning_fn );
+	png_structp png_ptr	= png_create_read_struct_2( PNG_LIBPNG_VER_STRING, this, FPngImageWrapper::user_error_fn, FPngImageWrapper::user_warning_fn, NULL, FPngImageWrapper::user_malloc, FPngImageWrapper::user_free);
 	check( png_ptr );
 
 	png_infop info_ptr	= png_create_info_struct( png_ptr );
@@ -354,7 +354,7 @@ bool FPngImageWrapper::LoadPNGHeader()
 		// thread safety
 		FScopeLock PNGLock(&GPNGSection);
 
-		png_structp png_ptr = png_create_read_struct( PNG_LIBPNG_VER_STRING, this, FPngImageWrapper::user_error_fn, FPngImageWrapper::user_warning_fn );
+		png_structp png_ptr = png_create_read_struct_2( PNG_LIBPNG_VER_STRING, this, FPngImageWrapper::user_error_fn, FPngImageWrapper::user_warning_fn, NULL, FPngImageWrapper::user_malloc, FPngImageWrapper::user_free );
 		check(png_ptr);
 
 		png_infop info_ptr = png_create_info_struct( png_ptr );
@@ -424,5 +424,17 @@ void FPngImageWrapper::user_warning_fn( png_structp png_ptr, png_const_charp war
 	UE_LOG(LogImageWrapper, Warning, TEXT("PNG Warning: %s"), ANSI_TO_TCHAR(warning_msg));
 }
 
+void* FPngImageWrapper::user_malloc( png_structp /*png_ptr*/, png_size_t size)
+{
+	check(size > 0);
+	return FMemory::Malloc(size);
+}
+
+void FPngImageWrapper::user_free(png_structp /*png_ptr*/, png_voidp struct_ptr )
+{
+	check(struct_ptr);
+	FMemory::Free(struct_ptr);
+}
+
 
 #endif	//WITH_UNREALPNG
diff --git a/Engine/Source/Developer/ImageWrapper/Private/PngImageWrapper.h b/Engine/Source/Developer/ImageWrapper/Private/PngImageWrapper.h
--- a/Engine/Source/Developer/ImageWrapper/Private/PngImageWrapper.h
+++ b/Engine/Source/Developer/ImageWrapper/Private/PngImageWrapper.h
@@ -62,11 +62,13 @@ public:
 protected:
 
 	// Callbacks for the pnglibs
-	static void user_read_compressed( png_structp png_ptr, png_bytep data, png_size_t length );
-	static void user_write_compressed( png_structp png_ptr, png_bytep data, png_size_t length );
-	static void user_flush_data( png_structp png_ptr );
-	static void user_error_fn( png_structp png_ptr, png_const_charp error_msg );
-	static void user_warning_fn( png_structp png_ptr, png_const_charp warning_msg );
+	static void  user_read_compressed( png_structp png_ptr, png_bytep data, png_size_t length );
+	static void  user_write_compressed( png_structp png_ptr, png_bytep data, png_size_t length );
+	static void  user_flush_data( png_structp png_ptr );
+	static void  user_error_fn( png_structp png_ptr, png_const_charp error_msg );
+	static void  user_warning_fn( png_structp png_ptr, png_const_charp warning_msg );
+	static void* user_malloc( png_structp png_ptr, png_size_t size);
+	static void  user_free(png_structp png_ptr, png_voidp struct_ptr );
 
 
 private:
