* added EXR texture import (at the moment 2D textures only, RGBA)

diff --git a/Engine/Source/Developer/ImageWrapper/Private/ExrImageWrapper.cpp b/Engine/Source/Developer/ImageWrapper/Private/ExrImageWrapper.cpp
--- a/Engine/Source/Developer/ImageWrapper/Private/ExrImageWrapper.cpp
+++ b/Engine/Source/Developer/ImageWrapper/Private/ExrImageWrapper.cpp
@@ -92,6 +92,77 @@
 };
 
 
+class FMemFileIn : public Imf::IStream
+{
+public:
+	//-------------------------------------------------------
+	// A constructor that opens the file with the given name.
+	// The destructor will close the file.
+	//-------------------------------------------------------
+
+	FMemFileIn(const void* InData, int32 InSize)
+		: Imf::IStream("")
+		, Data((const char *)InData)
+		, Size(InSize)
+		, Pos(0)
+	{
+	}
+
+	//------------------------------------------------------
+	// Read from the stream:
+	//
+	// read(c,n) reads n bytes from the stream, and stores
+	// them in array c.  If the stream contains less than n
+	// bytes, or if an I/O error occurs, read(c,n) throws
+	// an exception.  If read(c,n) reads the last byte from
+	// the file it returns false, otherwise it returns true.
+	//------------------------------------------------------
+
+	virtual bool read (char c[/*n*/], int n)
+	{
+		if(Pos + n > Size)
+		{
+			return false;
+		}
+
+		for (int32 i = 0; i < n; ++i)
+		{
+			c[i] = Data[Pos];
+			++Pos;
+		}
+
+		return Pos >= Size;
+	}
+
+	//--------------------------------------------------------
+	// Get the current reading position, in bytes from the
+	// beginning of the file.  If the next call to read() will
+	// read the first byte in the file, tellg() returns 0.
+	//--------------------------------------------------------
+
+	virtual Imf::Int64 tellg()
+	{
+		return Pos;
+	}
+
+	//-------------------------------------------
+	// Set the current reading position.
+	// After calling seekg(i), tellg() returns i.
+	//-------------------------------------------
+
+	virtual void seekg(Imf::Int64 pos)
+	{
+		Pos = pos;
+	}
+
+private:
+
+	const char* Data;
+	int32 Size;
+	int64 Pos;
+};
+
+
 namespace
 {
 	/////////////////////////////////////////
@@ -287,27 +358,70 @@
 
 void FExrImageWrapper::Uncompress( const ERGBFormat::Type InFormat, const int32 InBitDepth )
 {
-	checkf(false, TEXT("EXR decompression not supported"));
+	// Ensure we haven't already uncompressed the file.
+	if ( RawData.Num() != 0 )
+	{
+		return;
+	}
+
+	FMemFileIn MemFile(&CompressedData[0], CompressedData.Num());
+
+	Imf::RgbaInputFile ImfFile(MemFile);
+	Imath::Box2i win = ImfFile.dataWindow();
+
+	check(BitDepth == 16);
+	check(Width);
+	check(Height);
+
+	uint32 Channels = 4;
+
+	RawData.Empty();
+	RawData.AddUninitialized( Width * Height * Channels * (BitDepth / 8) );
+
+	int dx = win.min.x;
+	int dy = win.min.y;
+
+	ImfFile.setFrameBuffer((Imf::Rgba*)&RawData[0] - dx - dy * Width, 1, Width);
+	ImfFile.readPixels(win.min.y, win.max.y);
 }
 
+// from http://www.openexr.com/ReadingAndWritingImageFiles.pdf
+bool IsThisAnOpenExrFile(Imf::IStream& f)
+{
+	char b[4];
+	f.read(b, sizeof(b));
+
+	f.seekg(0);
+	return b[0] == 0x76 && b[1] == 0x2f && b[2] == 0x31 && b[3] == 0x01;}
+
 bool FExrImageWrapper::SetCompressed( const void* InCompressedData, int32 InCompressedSize )
 {
-	bool bResult = FImageWrapperBase::SetCompressed( InCompressedData, InCompressedSize );
+	if(!FImageWrapperBase::SetCompressed( InCompressedData, InCompressedSize))
+	{
+		return false;
+	}
 
-	return bResult;
+	FMemFileIn MemFile(InCompressedData, InCompressedSize);
+
+	if(!IsThisAnOpenExrFile(MemFile))
+	{
+		return false;
+	}
+
+	Imf::RgbaInputFile ImfFile(MemFile);
+	Imath::Box2i win = ImfFile.dataWindow();
+
+	Imath::V2i dim(win.max.x - win.min.x + 1, win.max.y - win.min.y + 1);
+
+	BitDepth = 16;
+
+	Width = dim.x;
+	Height = dim.y;
+
+	// ideally we can specify float here
+	Format = ERGBFormat::RGBA;
+
+	return true;
 }
 
-bool FExrImageWrapper::GetRaw( const ERGBFormat::Type InFormat, int32 InBitDepth, const TArray<uint8>*& OutRawData )
-{
-	LastError.Empty();
-	Uncompress(InFormat, InBitDepth);
-
-	if (LastError.IsEmpty())
-	{
-		OutRawData = &GetRawData();
-	}
-
-	return LastError.IsEmpty();
-}
-
 #endif // WITH_UNREALEXR
diff --git a/Engine/Source/Developer/ImageWrapper/Private/ExrImageWrapper.h b/Engine/Source/Developer/ImageWrapper/Private/ExrImageWrapper.h
--- a/Engine/Source/Developer/ImageWrapper/Private/ExrImageWrapper.h
+++ b/Engine/Source/Developer/ImageWrapper/Private/ExrImageWrapper.h
@@ -13,6 +13,7 @@
 #include "ThirdParty/openexr/Deploy/include/ImfHeader.h"
 #include "ThirdParty/openexr/Deploy/include/ImfStdIO.h"
 #include "ThirdParty/openexr/Deploy/include/ImfChannelList.h"
+#include "ThirdParty/openexr/Deploy/include/ImfRgbaFile.h"
 
 /**
  * OpenEXR implementation of the helper class
@@ -38,8 +39,6 @@ public:
 	
 	virtual bool SetCompressed( const void* InCompressedData, int32 InCompressedSize ) override;
 
-	virtual bool GetRaw( const ERGBFormat::Type InFormat, int32 InBitDepth, const TArray<uint8>*& OutRawData ) override;
-
 	// End FImageWrapper Interface
 
 private:
diff --git a/Engine/Source/Editor/UnrealEd/Private/Factories/EditorFactories.cpp b/Engine/Source/Editor/UnrealEd/Private/Factories/EditorFactories.cpp
--- a/Engine/Source/Editor/UnrealEd/Private/Factories/EditorFactories.cpp
+++ b/Engine/Source/Editor/UnrealEd/Private/Factories/EditorFactories.cpp
@@ -3235,6 +3235,7 @@ UTextureFactory::UTextureFactory(const FObjectInitializer& ObjectInitializer)
 	Formats.Add( TEXT( "png;Texture" ) );
 	Formats.Add( TEXT( "jpg;Texture" ) );
 	Formats.Add( TEXT( "jpeg;Texture" ) );
+	Formats.Add( TEXT( "exr;Texture" ) );
 
 	bCreateNew = false;
 	bEditorImport = true;
@@ -3587,6 +3588,68 @@ UTexture* UTextureFactory::ImportTexture(UClass* Class, UObject* InParent, FName
 		return Texture;
 	}
 	//
+	// EXR
+	//
+	IImageWrapperPtr ExrImageWrapper = ImageWrapperModule.CreateImageWrapper( EImageFormat::EXR );
+	if ( ExrImageWrapper.IsValid() && ExrImageWrapper->SetCompressed( Buffer, Length ) )
+	{
+		int32 Width = ExrImageWrapper->GetWidth();
+		int32 Height = ExrImageWrapper->GetHeight();
+
+		if ( !IsImportResolutionValid( Width, Height, bAllowNonPowerOfTwo, Warn ) )
+		{
+			return nullptr;
+		}
+
+		// Select the texture's source format
+		ETextureSourceFormat TextureFormat = TSF_Invalid;
+		int32 BitDepth = ExrImageWrapper->GetBitDepth();
+		ERGBFormat::Type Format = ExrImageWrapper->GetFormat();
+
+		if ( Format == ERGBFormat::RGBA && BitDepth == 16 )
+		{
+			TextureFormat = TSF_RGBA16F;
+			Format = ERGBFormat::BGRA;
+		}
+
+		if ( TextureFormat == TSF_Invalid )
+		{
+			Warn->Logf( ELogVerbosity::Error, TEXT( "EXR file contains data in an unsupported format." ) );
+			return nullptr;
+		}
+
+		UTexture2D* Texture = CreateTexture2D( InParent, Name, Flags );
+		if ( Texture )
+		{
+			const TArray<uint8>* Raw = nullptr;
+			if ( ExrImageWrapper->GetRaw( Format, BitDepth, Raw ) )
+			{
+				Texture->Source.Init(
+					Width,
+					Height,
+					/*NumSlices=*/ 1,
+					/*NumMips=*/ 1,
+					TextureFormat
+					);
+				Texture->SRGB = false;
+				Texture->CompressionSettings = TC_HDR;
+
+				uint8* MipData = Texture->Source.LockMip( 0 );
+				FMemory::Memcpy( MipData, Raw->GetData(), Raw->Num() );
+				Texture->Source.UnlockMip( 0 );
+			}
+			else
+			{
+				Warn->Logf( ELogVerbosity::Error, TEXT( "Failed to decode EXR." ) );
+				Texture->MarkPendingKill();
+
+				return nullptr;
+			}
+		}
+
+		return Texture;
+	}
+	//
 	// BMP
 	//
 	const FBitmapInfoHeader* bmhdr = (FBitmapInfoHeader *)(Buffer + sizeof(FBitmapFileHeader));
