Improved screen shot performance.

Changed PNG compression from smallest to fastest.

diff --git a/Engine/Source/Developer/ImageWrapper/Private/PngImageWrapper.cpp b/Engine/Source/Developer/ImageWrapper/Private/PngImageWrapper.cpp
--- a/Engine/Source/Developer/ImageWrapper/Private/PngImageWrapper.cpp
+++ b/Engine/Source/Developer/ImageWrapper/Private/PngImageWrapper.cpp
@@ -117,7 +117,7 @@ void FPngImageWrapper::Compress( int32 Quality )
 
 		PNGWriteGuard PNGGuard(&png_ptr, &info_ptr);
 		{
-			png_set_compression_level(png_ptr, Z_BEST_COMPRESSION);
+			png_set_compression_level(png_ptr, Z_BEST_SPEED);
 			png_set_IHDR(png_ptr, info_ptr, Width, Height, RawBitDepth, (RawFormat == ERGBFormat::Gray) ? PNG_COLOR_TYPE_GRAY : PNG_COLOR_TYPE_RGBA, PNG_INTERLACE_NONE, PNG_COMPRESSION_TYPE_DEFAULT, PNG_FILTER_TYPE_DEFAULT);
 			png_set_write_fn(png_ptr, this, FPngImageWrapper::user_write_compressed, FPngImageWrapper::user_flush_data);
 			png_bytep* row_pointers = (png_bytep*) png_malloc( png_ptr, Height*sizeof(png_bytep) );
