Prevent assert if icon decoding falls back to the Icns wrapper on Windows due to unsupported formats, etc..

diff --git a/Engine/Source/Developer/ImageWrapper/Private/IcnsImageWrapper.cpp b/Engine/Source/Developer/ImageWrapper/Private/IcnsImageWrapper.cpp
--- a/Engine/Source/Developer/ImageWrapper/Private/IcnsImageWrapper.cpp
+++ b/Engine/Source/Developer/ImageWrapper/Private/IcnsImageWrapper.cpp
@@ -11,6 +11,24 @@ FIcnsImageWrapper::FIcnsImageWrapper()
 {
 }
 
+bool FIcnsImageWrapper::SetCompressed( const void* InCompressedData, int32 InCompressedSize )
+{
+#if PLATFORM_MAC
+	return FImageWrapperBase::SetCompressed(InCompressedData, InCompressedSize);
+#else
+	return false;
+#endif
+}
+
+bool FIcnsImageWrapper::SetRaw( const void* InRawData, int32 InRawSize, const int32 InWidth, const int32 InHeight, const ERGBFormat::Type InFormat, const int32 InBitDepth )
+{
+#if PLATFORM_MAC
+	return FImageWrapperBase::SetRaw(InRawData, InRawSize, InWidth, InHeight, InFormat, InBitDepth);
+#else
+	return false;
+#endif
+}
+
 void FIcnsImageWrapper::Compress(int32 Quality)
 {
 	checkf(false, TEXT("ICNS compression not supported"));
