Missing header.

diff --git a/Engine/Source/Developer/ImageWrapper/Private/IcnsImageWrapper.h b/Engine/Source/Developer/ImageWrapper/Private/IcnsImageWrapper.h
--- a/Engine/Source/Developer/ImageWrapper/Private/IcnsImageWrapper.h
+++ b/Engine/Source/Developer/ImageWrapper/Private/IcnsImageWrapper.h
@@ -19,6 +19,9 @@ public:
 
 	// FImageWrapper Interface
 
+	virtual bool SetCompressed( const void* InCompressedData, int32 InCompressedSize ) override;
+	virtual bool SetRaw( const void* InRawData, int32 InRawSize, const int32 InWidth, const int32 InHeight, const ERGBFormat::Type InFormat, const int32 InBitDepth ) override;
+
 	virtual void Compress( int32 Quality ) override;
 	virtual void Uncompress( const ERGBFormat::Type InFormat, int32 InBitDepth ) override;
 };
