Guard against bad data being passed to image compressor methods.  Attemping to fix unreproducible import crashes

diff --git a/Engine/Source/Developer/ImageWrapper/Private/ImageWrapperBase.cpp b/Engine/Source/Developer/ImageWrapper/Private/ImageWrapperBase.cpp
--- a/Engine/Source/Developer/ImageWrapper/Private/ImageWrapperBase.cpp
+++ b/Engine/Source/Developer/ImageWrapper/Private/ImageWrapperBase.cpp
@@ -66,17 +66,19 @@ bool FImageWrapperBase::GetRaw( const ERGBFormat::Type InFormat, int32 InBitDept
 
 bool FImageWrapperBase::SetCompressed( const void* InCompressedData, int32 InCompressedSize )
 {
-	check(InCompressedData != NULL);
-	check(InCompressedSize > 0);
+	if( InCompressedSize > 0 && InCompressedData != nullptr )
+	{
+		Reset();
+		RawData.Empty();			// Invalidates the raw data too
 
-	Reset();
-	RawData.Empty();			// Invalidates the raw data too
+		CompressedData.Empty(InCompressedSize);
+		CompressedData.AddUninitialized(InCompressedSize);
+		FMemory::Memcpy(CompressedData.GetData(), InCompressedData, InCompressedSize);
 
-	CompressedData.Empty( InCompressedSize );
-	CompressedData.AddUninitialized( InCompressedSize );
-	FMemory::Memcpy(CompressedData.GetData(), InCompressedData, InCompressedSize);
+		return true;
+	}
 
-	return true;
+	return false;
 }
 
 
