Disable shadow variable warnings around third party includes that throw warnings

diff --git a/Engine/Plugins/Developer/PerforceSourceControl/Source/PerforceSourceControl/Private/PerforceSourceControlPrivatePCH.h b/Engine/Plugins/Developer/PerforceSourceControl/Source/PerforceSourceControl/Private/PerforceSourceControlPrivatePCH.h
--- a/Engine/Plugins/Developer/PerforceSourceControl/Source/PerforceSourceControl/Private/PerforceSourceControlPrivatePCH.h
+++ b/Engine/Plugins/Developer/PerforceSourceControl/Source/PerforceSourceControl/Private/PerforceSourceControlPrivatePCH.h
@@ -5,6 +5,10 @@
 #include "Core.h"
 #include "SlateBasics.h"
 #include "EditorStyle.h"
+
+PRAGMA_DISABLE_SHADOW_VARIABLE_WARNINGS
 #include <p4/clientapi.h>
 #include <p4/i18napi.h>
+PRAGMA_POP
+
 #include "ISourceControlModule.h"
\ No newline at end of file
diff --git a/Engine/Source/Developer/ImageWrapper/Private/ExrImageWrapper.h b/Engine/Source/Developer/ImageWrapper/Private/ExrImageWrapper.h
--- a/Engine/Source/Developer/ImageWrapper/Private/ExrImageWrapper.h
+++ b/Engine/Source/Developer/ImageWrapper/Private/ExrImageWrapper.h
@@ -4,6 +4,7 @@
 
 #if WITH_UNREALEXR
 
+PRAGMA_DISABLE_SHADOW_VARIABLE_WARNINGS
 #include "ThirdParty/openexr/Deploy/include/ImfIO.h"
 #include "ThirdParty/openexr/Deploy/include/ImathBox.h"
 #include "ThirdParty/openexr/Deploy/include/ImfChannelList.h"
@@ -14,6 +15,7 @@
 #include "ThirdParty/openexr/Deploy/include/ImfStdIO.h"
 #include "ThirdParty/openexr/Deploy/include/ImfChannelList.h"
 #include "ThirdParty/openexr/Deploy/include/ImfRgbaFile.h"
+PRAGMA_POP
 
 /**
  * OpenEXR implementation of the helper class
diff --git a/Engine/Source/Developer/ShaderFormatOpenGL/Private/GlslBackend.cpp b/Engine/Source/Developer/ShaderFormatOpenGL/Private/GlslBackend.cpp
--- a/Engine/Source/Developer/ShaderFormatOpenGL/Private/GlslBackend.cpp
+++ b/Engine/Source/Developer/ShaderFormatOpenGL/Private/GlslBackend.cpp
@@ -36,7 +36,11 @@
 #include "hlslcc_private.h"
 #include "GlslBackend.h"
 #include "compiler.h"
+
+PRAGMA_DISABLE_SHADOW_VARIABLE_WARNINGS
 #include "glsl_parser_extras.h"
+PRAGMA_POP
+
 #include "hash_table.h"
 #include "ir_rvalue_visitor.h"
 #include "PackUniformBuffers.h"
diff --git a/Engine/Source/Developer/iOS/MetalShaderFormat/Private/MetalBackend.cpp b/Engine/Source/Developer/iOS/MetalShaderFormat/Private/MetalBackend.cpp
--- a/Engine/Source/Developer/iOS/MetalShaderFormat/Private/MetalBackend.cpp
+++ b/Engine/Source/Developer/iOS/MetalShaderFormat/Private/MetalBackend.cpp
@@ -7,7 +7,11 @@
 #include "hlslcc_private.h"
 #include "MetalBackend.h"
 #include "compiler.h"
+
+PRAGMA_DISABLE_SHADOW_VARIABLE_WARNINGS
 #include "glsl_parser_extras.h"
+PRAGMA_POP
+
 #include "hash_table.h"
 #include "ir_rvalue_visitor.h"
 #include "PackUniformBuffers.h"
diff --git a/Engine/Source/Editor/UnrealEd/Private/EditorPhysXSupport.h b/Engine/Source/Editor/UnrealEd/Private/EditorPhysXSupport.h
--- a/Engine/Source/Editor/UnrealEd/Private/EditorPhysXSupport.h
+++ b/Engine/Source/Editor/UnrealEd/Private/EditorPhysXSupport.h
@@ -11,6 +11,8 @@
 #pragma warning( push )
 #pragma warning( disable : 4946 ) // reinterpret_cast used between related classes
 
+PRAGMA_DISABLE_SHADOW_VARIABLE_WARNINGS
+
 #if USING_CODE_ANALYSIS
 	#pragma warning( push )
 	#pragma warning( disable : ALL_CODE_ANALYSIS_WARNINGS )
@@ -64,6 +66,8 @@
 
 #pragma warning( pop )
 
+PRAGMA_POP
+
 using namespace physx;
 
 #endif // WITH_PHYSX
diff --git a/Engine/Source/Runtime/Core/Private/Internationalization/ICUText.cpp b/Engine/Source/Runtime/Core/Private/Internationalization/ICUText.cpp
--- a/Engine/Source/Runtime/Core/Private/Internationalization/ICUText.cpp
+++ b/Engine/Source/Runtime/Core/Private/Internationalization/ICUText.cpp
@@ -6,10 +6,14 @@
 #include "Text.h"
 #include "TextHistory.h"
 #include "ICUCulture.h"
+
+PRAGMA_DISABLE_SHADOW_VARIABLE_WARNINGS
 #include <unicode/coll.h>
 #include <unicode/sortkey.h>
 #include <unicode/numfmt.h>
 #include <unicode/msgfmt.h>
+PRAGMA_POP
+
 #include "ICUUtilities.h"
 #include "ICUTextCharacterIterator.h"
 
diff --git a/Engine/Source/Runtime/Engine/Public/PhysXIncludes.h b/Engine/Source/Runtime/Engine/Public/PhysXIncludes.h
--- a/Engine/Source/Runtime/Engine/Public/PhysXIncludes.h
+++ b/Engine/Source/Runtime/Engine/Public/PhysXIncludes.h
@@ -8,6 +8,8 @@
 MSVC_PRAGMA(warning(push))
 MSVC_PRAGMA(warning(disable : 4946)) // reinterpret_cast used between related classes
 
+PRAGMA_DISABLE_SHADOW_VARIABLE_WARNINGS
+
 #if USING_CODE_ANALYSIS
 	MSVC_PRAGMA(warning(push))
 	MSVC_PRAGMA(warning(disable : ALL_CODE_ANALYSIS_WARNINGS))
@@ -70,6 +72,8 @@ MSVC_PRAGMA(warning(disable : 4946)) // reinterpret_cast used between related cl
 
 MSVC_PRAGMA(warning(pop))
 
+PRAGMA_POP
+
 using namespace physx;
 
 #endif // WITH_PHYSX
\ No newline at end of file
