Optimizations to EXR exporter. For larger image sizes this is up to 25% quicker.

diff --git a/Engine/Source/Developer/ImageWrapper/Private/ExrImageWrapper.cpp b/Engine/Source/Developer/ImageWrapper/Private/ExrImageWrapper.cpp
--- a/Engine/Source/Developer/ImageWrapper/Private/ExrImageWrapper.cpp
+++ b/Engine/Source/Developer/ImageWrapper/Private/ExrImageWrapper.cpp
@@ -58,7 +58,7 @@ public:
 	{
 		for (int32 i = 0; i < n; ++i)
 		{
-			Data.Insert((uint8)c[i], Pos++);
+			Data[Pos++] = c[i];
 		}
 	}
 
@@ -315,8 +315,10 @@ void FExrImageWrapper::CompressRaw(const sourcetype* SrcData, bool bIgnoreAlpha)
 	}
 
 	FMemFileOut MemFile("");
-	Imf::FrameBuffer ImfFrameBuffer;
+	const int32 OutputPixelSize = ((OutputFormat == Imf::HALF) ? 2 : 4);
+	MemFile.Data.AddUninitialized(Width * Height * NumWriteComponents * OutputPixelSize);
 
+	Imf::FrameBuffer ImfFrameBuffer;
 	TArray<uint8> ChannelOutputBuffers[4];
 
 	for (uint32 Channel = 0; Channel < NumWriteComponents; Channel++)
@@ -328,7 +330,8 @@ void FExrImageWrapper::CompressRaw(const sourcetype* SrcData, bool bIgnoreAlpha)
 	ImfFile.setFrameBuffer(ImfFrameBuffer);
 	ImfFile.writePixels(Height);
 
-	CompressedData = MemFile.Data;
+	CompressedData.AddUninitialized(MemFile.tellp());
+	FMemory::Memcpy(CompressedData.GetData(), MemFile.Data.GetData(), MemFile.tellp());
 }
 
 void FExrImageWrapper::Compress( int32 Quality )
@@ -366,7 +369,8 @@ void FExrImageWrapper::Uncompress( const ERGBFormat::Type InFormat, const int32
 
 	FMemFileIn MemFile(&CompressedData[0], CompressedData.Num());
 
-	Imf::RgbaInputFile ImfFile(MemFile);
+	Imf::RgbaInputFile ImfFile(MemFile);
+
 	Imath::Box2i win = ImfFile.dataWindow();
 
 	check(BitDepth == 16);
@@ -392,7 +396,9 @@ bool IsThisAnOpenExrFile(Imf::IStream& f)
 	f.read(b, sizeof(b));
 
 	f.seekg(0);
-	return b[0] == 0x76 && b[1] == 0x2f && b[2] == 0x31 && b[3] == 0x01;}
+
+	return b[0] == 0x76 && b[1] == 0x2f && b[2] == 0x31 && b[3] == 0x01;
+}
 
 bool FExrImageWrapper::SetCompressed( const void* InCompressedData, int32 InCompressedSize )
 {
@@ -408,7 +414,8 @@ bool FExrImageWrapper::SetCompressed( const void* InCompressedData, int32 InComp
 		return false;
 	}
 
-	Imf::RgbaInputFile ImfFile(MemFile);
+	Imf::RgbaInputFile ImfFile(MemFile);
+
 	Imath::Box2i win = ImfFile.dataWindow();
 
 	Imath::V2i dim(win.max.x - win.min.x + 1, win.max.y - win.min.y + 1);
