Minor optimization to HDR export. Added some logging (Enable verbose LogImageWrapper to view)

diff --git a/Engine/Source/Developer/ImageWrapper/Private/ExrImageWrapper.cpp b/Engine/Source/Developer/ImageWrapper/Private/ExrImageWrapper.cpp
--- a/Engine/Source/Developer/ImageWrapper/Private/ExrImageWrapper.cpp
+++ b/Engine/Source/Developer/ImageWrapper/Private/ExrImageWrapper.cpp
@@ -58,8 +58,9 @@ public:
 	{
 		for (int32 i = 0; i < n; ++i)
 		{
-			Data[Pos++] = c[i];
+			Data[Pos + i] = c[i];
 		}
+		Pos += n;
 	}
 
 
@@ -301,6 +302,7 @@ void FExrImageWrapper::WriteFrameBufferChannel(Imf::FrameBuffer& ImfFrameBuffer,
 template <Imf::PixelType OutputFormat, typename sourcetype>
 void FExrImageWrapper::CompressRaw(const sourcetype* SrcData, bool bIgnoreAlpha)
 {
+	const double StartTime = FPlatformTime::Seconds();
 	uint32 NumWriteComponents = GetNumChannelsFromFormat(RawFormat);
 	if (bIgnoreAlpha && NumWriteComponents == 4)
 	{
@@ -332,6 +334,9 @@ void FExrImageWrapper::CompressRaw(const sourcetype* SrcData, bool bIgnoreAlpha)
 
 	CompressedData.AddUninitialized(MemFile.tellp());
 	FMemory::Memcpy(CompressedData.GetData(), MemFile.Data.GetData(), MemFile.tellp());
+
+	const double DeltaTime = FPlatformTime::Seconds() - StartTime;
+	UE_LOG(LogImageWrapper, Verbose, TEXT("Compressed image in %.3f seconds"), DeltaTime);
 }
 
 void FExrImageWrapper::Compress( int32 Quality )
