PR #1492: Update PngImageWrapper.cpp (Contributed by [REDACTED_PERSON])

diff --git a/Engine/Source/Developer/ImageWrapper/Private/PngImageWrapper.cpp b/Engine/Source/Developer/ImageWrapper/Private/PngImageWrapper.cpp
--- a/Engine/Source/Developer/ImageWrapper/Private/PngImageWrapper.cpp
+++ b/Engine/Source/Developer/ImageWrapper/Private/PngImageWrapper.cpp
@@ -298,7 +298,7 @@ void FPngImageWrapper::UncompressPNGData( const ERGBFormat::Type InFormat, const
 		{
 #if PNG_LIBPNG_VER >= 10504
 			check(0); // Needs testing
-			Transform |= PNG_TRANSFORM_EXPAND_16
+			Transform |= PNG_TRANSFORM_EXPAND_16;
 #else
 			// Expanding 8-bit images to 16-bit via transform needs a libpng update
 			check(0);
