Fixed error handling in FPngImageWrapper.

diff --git a/Engine/Source/Developer/ImageWrapper/Private/PngImageWrapper.cpp b/Engine/Source/Developer/ImageWrapper/Private/PngImageWrapper.cpp
--- a/Engine/Source/Developer/ImageWrapper/Private/PngImageWrapper.cpp
+++ b/Engine/Source/Developer/ImageWrapper/Private/PngImageWrapper.cpp
@@ -144,7 +144,14 @@ void FPngImageWrapper::Compress( int32 Quality )
 			}
 #endif
 
-			png_write_png(png_ptr, info_ptr, Transform, NULL);
+			if (!setjmp(SetjmpBuffer))
+			{
+				png_write_png(png_ptr, info_ptr, Transform, NULL);
+			}
+			else
+			{
+				png_destroy_write_struct(&png_ptr, &info_ptr);
+			}
 		}
 	}
 }
@@ -305,7 +312,14 @@ void FPngImageWrapper::UncompressPNGData( const ERGBFormat::Type InFormat, const
 #endif
 		}
 
-		png_read_png(png_ptr, info_ptr, Transform, NULL);
+		if (!setjmp(SetjmpBuffer))
+		{
+			png_read_png(png_ptr, info_ptr, Transform, NULL);
+		}
+		else
+		{
+			png_destroy_read_struct(&png_ptr, &info_ptr, NULL);
+		}
 	}
 
 	RawFormat = InFormat;
@@ -404,6 +418,8 @@ void FPngImageWrapper::user_error_fn( png_structp png_ptr, png_const_charp error
 	ctx->SetError(*ErrorMsg);
 
 	UE_LOG(LogImageWrapper, Error, TEXT("PNG Error: %s"), *ErrorMsg);
+
+	longjmp(ctx->SetjmpBuffer, 1);
 }
 
 
diff --git a/Engine/Source/Developer/ImageWrapper/Private/PngImageWrapper.h b/Engine/Source/Developer/ImageWrapper/Private/PngImageWrapper.h
--- a/Engine/Source/Developer/ImageWrapper/Private/PngImageWrapper.h
+++ b/Engine/Source/Developer/ImageWrapper/Private/PngImageWrapper.h
@@ -8,7 +8,7 @@
 
 #include "ThirdParty/libPNG/libPNG-1.5.2/png.h"
 #include "ThirdParty/libPNG/libPNG-1.5.2/pnginfo.h"
-
+#include <setjmp.h>
 
 /**
  * PNG implementation of the helper class.
@@ -83,6 +83,9 @@ private:
 
 	// The number of channels.
 	uint8 Channels;
+
+	// setjmp buffer for error recovery
+	jmp_buf SetjmpBuffer;
 };
 
 
