Disable VS2013 compiler warning due to use of setjmp/longjmp (and fix a memory leak as a result of a destructor not being called).

diff --git a/Engine/Source/Developer/ImageWrapper/Private/PngImageWrapper.cpp b/Engine/Source/Developer/ImageWrapper/Private/PngImageWrapper.cpp
--- a/Engine/Source/Developer/ImageWrapper/Private/PngImageWrapper.cpp
+++ b/Engine/Source/Developer/ImageWrapper/Private/PngImageWrapper.cpp
@@ -5,6 +5,9 @@
 
 #if WITH_UNREALPNG
 
+// Disable warning "interaction between '_setjmp' and C++ object destruction is non-portable"
+#pragma warning(push)
+#pragma warning(disable:4611)
 
 /** Only allow one thread to use libpng at a time (it's not thread safe) */
 FCriticalSection GPNGSection;
@@ -148,10 +151,6 @@ void FPngImageWrapper::Compress( int32 Quality )
 			{
 				png_write_png(png_ptr, info_ptr, Transform, NULL);
 			}
-			else
-			{
-				png_destroy_write_struct(&png_ptr, &info_ptr);
-			}
 		}
 	}
 }
@@ -316,10 +315,6 @@ void FPngImageWrapper::UncompressPNGData( const ERGBFormat::Type InFormat, const
 		{
 			png_read_png(png_ptr, info_ptr, Transform, NULL);
 		}
-		else
-		{
-			png_destroy_read_struct(&png_ptr, &info_ptr, NULL);
-		}
 	}
 
 	RawFormat = InFormat;
@@ -413,11 +408,15 @@ void FPngImageWrapper::user_flush_data( png_structp png_ptr )
 
 void FPngImageWrapper::user_error_fn( png_structp png_ptr, png_const_charp error_msg )
 {
-	FPngImageWrapper* ctx = (FPngImageWrapper*)png_get_io_ptr( png_ptr );
-	FString ErrorMsg = ANSI_TO_TCHAR(error_msg);
-	ctx->SetError(*ErrorMsg);
+	FPngImageWrapper* ctx = (FPngImageWrapper*)png_get_io_ptr(png_ptr);
 
-	UE_LOG(LogImageWrapper, Error, TEXT("PNG Error: %s"), *ErrorMsg);
+	{
+		FString ErrorMsg = ANSI_TO_TCHAR(error_msg);
+		ctx->SetError(*ErrorMsg);
+
+		UE_LOG(LogImageWrapper, Error, TEXT("PNG Error: %s"), *ErrorMsg);
+	}
+	// Ensure that FString is destructed prior to executing the longjmp
 
 	longjmp(ctx->SetjmpBuffer, 1);
 }
@@ -440,5 +439,7 @@ void FPngImageWrapper::user_free(png_structp /*png_ptr*/, png_voidp struct_ptr )
 	FMemory::Free(struct_ptr);
 }
 
+// Renable warning "interaction between '_setjmp' and C++ object destruction is non-portable"
+#pragma warning(pop)
 
 #endif	//WITH_UNREALPNG
