Fixed non-Windows builds

diff --git a/Engine/Source/Developer/ImageWrapper/Private/PngImageWrapper.cpp b/Engine/Source/Developer/ImageWrapper/Private/PngImageWrapper.cpp
--- a/Engine/Source/Developer/ImageWrapper/Private/PngImageWrapper.cpp
+++ b/Engine/Source/Developer/ImageWrapper/Private/PngImageWrapper.cpp
@@ -6,8 +6,10 @@
 #if WITH_UNREALPNG
 
 // Disable warning "interaction between '_setjmp' and C++ object destruction is non-portable"
-#pragma warning(push)
-#pragma warning(disable:4611)
+#if _MSC_VER
+	#pragma warning(push)
+	#pragma warning(disable:4611)
+#endif
 
 /** Only allow one thread to use libpng at a time (it's not thread safe) */
 FCriticalSection GPNGSection;
@@ -440,6 +442,8 @@ void FPngImageWrapper::user_free(png_structp /*png_ptr*/, png_voidp struct_ptr )
 }
 
 // Renable warning "interaction between '_setjmp' and C++ object destruction is non-portable"
-#pragma warning(pop)
+#if _MSC_VER
+	#pragma warning(pop)
+#endif
 
 #endif	//WITH_UNREALPNG
