Properly handle BGRA format in FIcnsImageWrapper

diff --git a/Engine/Source/Developer/ImageWrapper/Private/IcnsImageWrapper.cpp b/Engine/Source/Developer/ImageWrapper/Private/IcnsImageWrapper.cpp
--- a/Engine/Source/Developer/ImageWrapper/Private/IcnsImageWrapper.cpp
+++ b/Engine/Source/Developer/ImageWrapper/Private/IcnsImageWrapper.cpp
@@ -46,8 +46,7 @@ void FIcnsImageWrapper::Uncompress(const ERGBFormat::Type InFormat, const int32
 		NSBitmapImageRep* Bitmap = [NSBitmapImageRep imageRepWithData:[Image TIFFRepresentation]];
 		if (Bitmap)
 		{
-			// @todo: Only 8-bit BGRA supported currently
-			check(InFormat == ERGBFormat::BGRA);
+			check(InFormat == ERGBFormat::BGRA || InFormat == ERGBFormat::RGBA);
 			check(InBitDepth == 8);
 
 			RawData.Empty();
@@ -58,6 +57,16 @@ void FIcnsImageWrapper::Uncompress(const ERGBFormat::Type InFormat, const int32
 
 			Width = [Bitmap pixelsWide];
 			Height = [Bitmap pixelsHigh];
+
+			if (Format == ERGBFormat::BGRA)
+			{
+				for (int32 Index = 0; Index < [Bitmap bytesPerPlane]; Index += 4)
+				{
+					uint8 Byte = RawData[Index];
+					RawData[Index] = RawData[Index + 2];
+					RawData[Index + 2] = Byte;
+				}
+			}
 		}
 		[Image release];
 	}
